Enable Modeling Mode in various Game Project Templates
#rnx

diff --git a/Templates/TP_Blank/TP_Blank.uproject b/Templates/TP_Blank/TP_Blank.uproject
--- a/Templates/TP_Blank/TP_Blank.uproject
+++ b/Templates/TP_Blank/TP_Blank.uproject
@@ -9,5 +9,14 @@
 			"Type": "Runtime",
 			"LoadingPhase": "Default"
 		}
+	],
+	"Plugins": [
+		{
+			"Name": "ModelingToolsEditorMode",
+			"Enabled": true,
+			"TargetAllowList": [
+				"Editor"
+			]
+		}
 	]
 }
\ No newline at end of file
diff --git a/Templates/TP_BlankBP/TP_BlankBP.uproject b/Templates/TP_BlankBP/TP_BlankBP.uproject
--- a/Templates/TP_BlankBP/TP_BlankBP.uproject
+++ b/Templates/TP_BlankBP/TP_BlankBP.uproject
@@ -3,5 +3,14 @@
 	"EngineAssociation": "",
 	"Category": "",
 	"Description": "",
-	"Modules": []
+	"Modules": [],
+	"Plugins": [
+		{
+			"Name": "ModelingToolsEditorMode",
+			"Enabled": true,
+			"TargetAllowList": [
+				"Editor"
+			]
+		}
+	]
 }
\ No newline at end of file
diff --git a/Templates/TP_FirstPerson/TP_FirstPerson.uproject b/Templates/TP_FirstPerson/TP_FirstPerson.uproject
--- a/Templates/TP_FirstPerson/TP_FirstPerson.uproject
+++ b/Templates/TP_FirstPerson/TP_FirstPerson.uproject
@@ -9,5 +9,14 @@
 			"Type": "Runtime",
 			"LoadingPhase": "Default"
 		}
-	]
+	],
+	"Plugins": [
+		{
+			"Name": "ModelingToolsEditorMode",
+			"Enabled": true,
+			"TargetAllowList": [
+				"Editor"
+			]
+		}
+	]	
 }
\ No newline at end of file
diff --git a/Templates/TP_FirstPersonBP/TP_FirstPersonBP.uproject b/Templates/TP_FirstPersonBP/TP_FirstPersonBP.uproject
--- a/Templates/TP_FirstPersonBP/TP_FirstPersonBP.uproject
+++ b/Templates/TP_FirstPersonBP/TP_FirstPersonBP.uproject
@@ -2,5 +2,14 @@
 	"FileVersion": 3,
 	"EngineAssociation": "",
 	"Category": "",
-	"Description": ""
+	"Description": "",
+	"Plugins": [
+		{
+			"Name": "ModelingToolsEditorMode",
+			"Enabled": true,
+			"TargetAllowList": [
+				"Editor"
+			]
+		}
+	]	
 }
\ No newline at end of file
diff --git a/Templates/TP_Puzzle/TP_Puzzle.uproject b/Templates/TP_Puzzle/TP_Puzzle.uproject
--- a/Templates/TP_Puzzle/TP_Puzzle.uproject
+++ b/Templates/TP_Puzzle/TP_Puzzle.uproject
@@ -9,5 +9,14 @@
 			"Type": "Runtime",
 			"LoadingPhase": "Default"
 		}
-	]
+	],
+	"Plugins": [
+		{
+			"Name": "ModelingToolsEditorMode",
+			"Enabled": true,
+			"TargetAllowList": [
+				"Editor"
+			]
+		}
+	]	
 }
\ No newline at end of file
diff --git a/Templates/TP_PuzzleBP/TP_PuzzleBP.uproject b/Templates/TP_PuzzleBP/TP_PuzzleBP.uproject
--- a/Templates/TP_PuzzleBP/TP_PuzzleBP.uproject
+++ b/Templates/TP_PuzzleBP/TP_PuzzleBP.uproject
@@ -2,5 +2,14 @@
 	"FileVersion": 3,
 	"EngineAssociation": "",
 	"Category": "",
-	"Description": ""
+	"Description": "",
+	"Plugins": [
+		{
+			"Name": "ModelingToolsEditorMode",
+			"Enabled": true,
+			"TargetAllowList": [
+				"Editor"
+			]
+		}
+	]	
 }
\ No newline at end of file
diff --git a/Templates/TP_ThirdPerson/TP_ThirdPerson.uproject b/Templates/TP_ThirdPerson/TP_ThirdPerson.uproject
--- a/Templates/TP_ThirdPerson/TP_ThirdPerson.uproject
+++ b/Templates/TP_ThirdPerson/TP_ThirdPerson.uproject
@@ -9,5 +9,14 @@
 			"Type": "Runtime",
 			"LoadingPhase": "Default"
 		}
-	]
+	],
+	"Plugins": [
+		{
+			"Name": "ModelingToolsEditorMode",
+			"Enabled": true,
+			"TargetAllowList": [
+				"Editor"
+			]
+		}
+	]	
 }
\ No newline at end of file
diff --git a/Templates/TP_ThirdPersonBP/TP_ThirdPersonBP.uproject b/Templates/TP_ThirdPersonBP/TP_ThirdPersonBP.uproject
--- a/Templates/TP_ThirdPersonBP/TP_ThirdPersonBP.uproject
+++ b/Templates/TP_ThirdPersonBP/TP_ThirdPersonBP.uproject
@@ -2,5 +2,14 @@
 	"FileVersion": 3,
 	"EngineAssociation": "",
 	"Category": "",
-	"Description": ""
+	"Description": "",
+	"Plugins": [
+		{
+			"Name": "ModelingToolsEditorMode",
+			"Enabled": true,
+			"TargetAllowList": [
+				"Editor"
+			]
+		}
+	]	
 }
\ No newline at end of file
diff --git a/Templates/TP_TopDown/TP_TopDown.uproject b/Templates/TP_TopDown/TP_TopDown.uproject
--- a/Templates/TP_TopDown/TP_TopDown.uproject
+++ b/Templates/TP_TopDown/TP_TopDown.uproject
@@ -9,5 +9,14 @@
 			"Type": "Runtime",
 			"LoadingPhase": "Default"
 		}
-	]
+	],
+	"Plugins": [
+		{
+			"Name": "ModelingToolsEditorMode",
+			"Enabled": true,
+			"TargetAllowList": [
+				"Editor"
+			]
+		}
+	]	
 }
\ No newline at end of file
diff --git a/Templates/TP_TopDownBP/TP_TopDownBP.uproject b/Templates/TP_TopDownBP/TP_TopDownBP.uproject
--- a/Templates/TP_TopDownBP/TP_TopDownBP.uproject
+++ b/Templates/TP_TopDownBP/TP_TopDownBP.uproject
@@ -2,5 +2,14 @@
 	"FileVersion": 3,
 	"EngineAssociation": "",
 	"Category": "",
-	"Description": ""
+	"Description": "",
+	"Plugins": [
+		{
+			"Name": "ModelingToolsEditorMode",
+			"Enabled": true,
+			"TargetAllowList": [
+				"Editor"
+			]
+		}
+	]	
 }
\ No newline at end of file
diff --git a/Templates/TP_VehicleAdv/TP_VehicleAdv.uproject b/Templates/TP_VehicleAdv/TP_VehicleAdv.uproject
--- a/Templates/TP_VehicleAdv/TP_VehicleAdv.uproject
+++ b/Templates/TP_VehicleAdv/TP_VehicleAdv.uproject
@@ -22,6 +22,13 @@
 		{
 			"Name": "RawInput",
 			"Enabled": true
-		}
+		},
+		{
+			"Name": "ModelingToolsEditorMode",
+			"Enabled": true,
+			"TargetAllowList": [
+				"Editor"
+			]
+		}		
 	]
 }
\ No newline at end of file
diff --git a/Templates/TP_VehicleAdvBP/TP_VehicleAdvBP.uproject b/Templates/TP_VehicleAdvBP/TP_VehicleAdvBP.uproject
--- a/Templates/TP_VehicleAdvBP/TP_VehicleAdvBP.uproject
+++ b/Templates/TP_VehicleAdvBP/TP_VehicleAdvBP.uproject
@@ -11,6 +11,13 @@
 		{
 			"Name": "RawInput",
 			"Enabled": true
-		}
+		},
+		{
+			"Name": "ModelingToolsEditorMode",
+			"Enabled": true,
+			"TargetAllowList": [
+				"Editor"
+			]
+		}		
 	]
 }
\ No newline at end of file
