Fixed weapon references
Fixed blueprint error on BP_Pickup
Removed external actors and world partition from first person maps. Should fix ddc warning on engine build.
Updated references in VR

diff --git a/Templates/TP_FirstPersonBP/manifest.json b/Templates/TP_FirstPersonBP/manifest.json
--- a/Templates/TP_FirstPersonBP/manifest.json
+++ b/Templates/TP_FirstPersonBP/manifest.json
@@ -57,7 +57,7 @@
 	],
 	"AdditionalFeaturePacks": [
 		{
-			"MountName": "FPWeapon",
+			"MountName": "Weapons",
 			"DetailLevels": [
 				"Standard"
 			]
@@ -69,9 +69,9 @@
 			]
 		},
 		{
-			"MountName": "FirstPersonArms",
+			"MountName": "Characters",
 			"DetailLevels": [
-				"Standard"
+				"High"
 			]
 		}
 	]
