TP_FirstPersonBP - updated content pack.
#rnx

diff --git a/Templates/TP_FirstPersonBP/manifest.json b/Templates/TP_FirstPersonBP/manifest.json
--- a/Templates/TP_FirstPersonBP/manifest.json
+++ b/Templates/TP_FirstPersonBP/manifest.json
@@ -57,19 +57,19 @@
 	],
 	"AdditionalFeaturePacks": [
 		{
-			"MountName": "Input",
+			"MountName": "LevelPrototyping",
 			"DetailLevels": [
 				"High"
 			]
 		},
 		{
-			"MountName": "LevelPrototyping",
+			"MountName": "Characters",
 			"DetailLevels": [
 				"High"
 			]
 		},
 		{
-			"MountName": "Characters",
+			"MountName": "Input",
 			"DetailLevels": [
 				"High"
 			]
